Fix PlayFabParty build errors with October 2025 GDK

#Synced-CL 51337133

diff --git a/Engine/Plugins/Online/Microsoft/PlayFabParty/Source/Private/PlayFabPartyManager.cpp b/Engine/Plugins/Online/Microsoft/PlayFabParty/Source/Private/PlayFabPartyManager.cpp
--- a/Engine/Plugins/Online/Microsoft/PlayFabParty/Source/Private/PlayFabPartyManager.cpp
+++ b/Engine/Plugins/Online/Microsoft/PlayFabParty/Source/Private/PlayFabPartyManager.cpp
@@ -26,10 +26,9 @@
 #include "Policies/CondensedJsonPrintPolicy.h"
 #include "Serialization/JsonWriter.h"
 #include "Templates/UnrealTemplate.h"
+#include <vector>
+
 
-THIRD_PARTY_INCLUDES_START
-#include <grdk.h>
-THIRD_PARTY_INCLUDES_END
 
 using FJsonWriter = TSharedRef<TJsonWriter<TCHAR, TCondensedJsonPrintPolicy<TCHAR>>>;
 using FJsonWriterFactory = TJsonWriterFactory<TCHAR, TCondensedJsonPrintPolicy<TCHAR>>;
@@ -61,9 +60,13 @@ FPlayFabPartyManager::~FPlayFabPartyManager()
 {
 	UnbindSystemDelegates();
 
+#if _GRDK_EDITION >= 251000 && !WITH_LEGACY_GDK_FOLDER_STRUCTURE
+	if (bIsPlayFabPartyInit || bIsPlayFabPartyXboxInit || bIsPlayFabCoreInit)
+#else
 	if (bIsPlayFabPartyInit || bIsPlayFabPartyXboxInit)
+#endif
 	{
-		ShutdownPlayFabParty();
+		ShutdownPlayFabParty(true);
 	}
 }
 
@@ -136,7 +139,11 @@ TOptional<FString> FPlayFabPartyManager::GetAppId()
 
 bool FPlayFabPartyManager::IsInitialized() const
 {
-	return bIsPlayFabPartyInit && bIsPlayFabPartyXboxInit;
+#if _GRDK_EDITION >= 251000 && !WITH_LEGACY_GDK_FOLDER_STRUCTURE
+	return bIsPlayFabCoreInit && bIsPlayFabPartyInit && bIsPlayFabPartyXboxInit;
+#else
+	return bIsPlayFabPartyInit && bIsPlayFabPartyXboxInit; 
+#endif
 }
 
 bool FPlayFabPartyManager::IsReady() const
@@ -302,13 +309,33 @@ PartyString FPlayFabPartyManager::GetEntityTokenForXuid(const uint64 Xuid)
 	return nullptr;
 }
 
+#if _GRDK_EDITION >= 251000 && !WITH_LEGACY_GDK_FOLDER_STRUCTURE
+PFEntityHandle FPlayFabPartyManager::GetEntityHandleForXuid(const uint64 Xuid)
+{
+	if (PFEntityHandle* EntityHandle = XboxUserIdToPlayFabEntityHandleMap.Find(Xuid))
+	{
+		return *EntityHandle;
+	}
+
+	return nullptr;
+}
+#endif
+
 
 FPlayFabPartyManager::FPlayFabPartyManager(FPlayFabPartySocketSubsystem& InSocketSubsystem, const FString& InAppId)
 	: SocketSubsystem(InSocketSubsystem)
 	, AppId(InAppId)
 {
 	GConfig->GetDouble(TEXT("PlayFab"), TEXT("LoginFailureDelaySeconds"), LoginFailureDelaySeconds, GEngineIni);
+	if (!GConfig->GetString(TEXT("PlayFab"), TEXT("PlayFabURL"), PlayFabURL, GEngineIni))
+	{
+		auto ConvertedAppId = StringCast<ANSICHAR>(*AppId);
+		PlayFabURL = FString::Printf(TEXT("https://%hs.playfabapi.com"), ConvertedAppId.Get());
+	}
+	
+#if _GRDK_EDITION < 251000 || WITH_LEGACY_GDK_FOLDER_STRUCTURE // Token refresh is handled automatically after 251000
 	GConfig->GetDouble(TEXT("PlayFab"), TEXT("LoginRefreshDelaySeconds"), LoginRefreshDelaySeconds, GEngineIni);
+#endif
 }
 
 void FPlayFabPartyManager::UpdatePlayFabPartyStatus()
@@ -334,78 +361,357 @@ void FPlayFabPartyManager::UpdatePlayFabPartyStatus()
 void FPlayFabPartyManager::InitPlayFabParty()
 {
 	GDK_SCOPE_NOT_TIME_SENSITIVE(); // (startup only) PartyManager & PartyXblManager Initialize call many functions which are not safe to call on time-sensitive threads
-
 	UE_LOG(LogPlayFabParty, Log, TEXT("Initializing PlayFabParty"));
 
 	check(!bIsPlayFabPartyInit);
 	check(!bIsPlayFabPartyXboxInit);
 
-	auto ConvertedAppId = StringCast<ANSICHAR>(*AppId);
+	// This is needed in order to pass task context to the queue, it can't be captured by the lambda since only non-capturing lambdas can implicitly convert to a raw function pointer.
+	TWeakPtr<FPlayFabPartyManager>* WeakThisPtr = new TWeakPtr<FPlayFabPartyManager>();
+	*WeakThisPtr = this->AsWeak();
 
-#if _GRDK_EDITION < 251000 || WITH_LEGACY_GDK_FOLDER_STRUCTURE
-	const PartyError InitializePlayFabErrorCode = Party::PartyManager::GetSingleton().Initialize(ConvertedAppId.Get());
+	HRESULT InitPFPartySubmitTaskResult = XTaskQueueSubmitCallback(
+		FGDKAsyncTaskQueue::GetGenericSerialQueue(),
+		XTaskQueuePort::Work,
+		WeakThisPtr,
+		[](void* context, bool cancel)
+		{
+
+			TWeakPtr<FPlayFabPartyManager>* WeakPFPartyManagerPtr = static_cast<TWeakPtr<FPlayFabPartyManager>*>(context);
+			TSharedPtr<FPlayFabPartyManager> SharedPFPartyManagerPtr = WeakPFPartyManagerPtr->Pin();
+			if (!SharedPFPartyManagerPtr.IsValid())
+			{
+				delete WeakPFPartyManagerPtr;
+				UE_LOG(LogPlayFabParty, Error, TEXT("[InitPlayFabParty PFInitialize] PlayFabManager was destroyed before the async block was executed"));
+				return;
+			}
+
+			if (SharedPFPartyManagerPtr->IsInitialized())
+			{
+				delete WeakPFPartyManagerPtr;
+				UE_LOG(LogPlayFabParty, Log, TEXT("[InitPlayFabParty PFInitialize] PlayFabManager was already initialized"));
+				return;
+			}
+
+			auto ConvertedAppId = StringCast<ANSICHAR>(*SharedPFPartyManagerPtr->AppId);
+			
+#if _GRDK_EDITION >= 251000 && !WITH_LEGACY_GDK_FOLDER_STRUCTURE
+
+			HRESULT PFCoreInitResult = PFInitialize(nullptr);
+			if (FAILED(PFCoreInitResult))
+			{
+				delete WeakPFPartyManagerPtr;
+				UE_LOG(LogPlayFabParty, Error, TEXT("Failed to initialize PlayFabCore. ErrorCode=[%0.8x] ErrorMessage=[%s]"), PFCoreInitResult, *GetPlayFabPartyErrorMessage(PFCoreInitResult));
+				return;
+			}
+
+			PFServiceConfigHandle LocalServiceConfigHandle = { nullptr };
+			HRESULT PFCoreConfigHandleResult = PFServiceConfigCreateHandle(TCHAR_TO_ANSI(*SharedPFPartyManagerPtr->PlayFabURL), ConvertedAppId.Get(), &LocalServiceConfigHandle);
+			if (FAILED(PFCoreConfigHandleResult))
+			{
+				UE_LOG(LogPlayFabParty, Error, TEXT("Failed to initialize PlayFabCore config handle. ErrorCode=[%0.8x] ErrorMessage=[%s]"), PFCoreConfigHandleResult, *GetPlayFabPartyErrorMessage(PFCoreConfigHandleResult));
+
+				SharedPFPartyManagerPtr->bIsPlayFabCoreInit = false;
+			}
+			else
+			{
+				SharedPFPartyManagerPtr->PFCoreServiceHandle.store(LocalServiceConfigHandle);
+				SharedPFPartyManagerPtr->bIsPlayFabCoreInit = true;
+			}
+
+			const Party::PartyInitializationConfiguration PartyInitConfig =
+			{
+				.titleId = ConvertedAppId.Get(),
+			};
+
+			const PartyError InitializePlayFabErrorCode = Party::PartyManager::GetSingleton().Initialize(&PartyInitConfig);
 #else
-	const Party::PartyInitializationConfiguration PartyInitConfig = 
-	{
-		.titleId = ConvertedAppId.Get(),
-	};
-	const PartyError InitializePlayFabErrorCode = Party::PartyManager::GetSingleton().Initialize(&PartyInitConfig);
+
+			const PartyError InitializePlayFabErrorCode = Party::PartyManager::GetSingleton().Initialize(ConvertedAppId.Get());
+#endif
+
+			
+			if (PARTY_FAILED(InitializePlayFabErrorCode))
+			{
+				UE_LOG(LogPlayFabParty, Warning, TEXT("Failed to initialize PlayFabParty. ErrorCode=[%u] Error=[%s]"), InitializePlayFabErrorCode, *GetPlayFabPartyErrorMessage(InitializePlayFabErrorCode));
+				// We don't try to try again because config is probably invalid if this fails
+				SharedPFPartyManagerPtr->bIsPlayFabPartyInit = false;
+			}
+			else
+			{
+				SharedPFPartyManagerPtr->bIsPlayFabPartyInit = true;
+			}
+
+			const PartyError InitializePlayFabXboxErrorCode = Party::PartyXblManager::GetSingleton().Initialize(ConvertedAppId.Get());
+			if (PARTY_FAILED(InitializePlayFabXboxErrorCode))
+			{
+				UE_LOG(LogPlayFabParty, Warning, TEXT("Failed to initialize PlayFabParty Xbox. ErrorCode=[%u] Error=[%s]"), InitializePlayFabXboxErrorCode, *GetPlayFabPartyErrorMessage(InitializePlayFabXboxErrorCode));
+				// We don't try to try again because config is probably invalid if this fails
+				SharedPFPartyManagerPtr->bIsPlayFabPartyXboxInit = false;
+			}
+			else
+			{
+				SharedPFPartyManagerPtr->bIsPlayFabPartyXboxInit = true;
+			}
+
+			HRESULT InitPFPartySubmitTaskResult = XTaskQueueSubmitCallback(
+				FGDKAsyncTaskQueue::GetGenericSerialQueue(),
+				XTaskQueuePort::Completion,
+				WeakPFPartyManagerPtr,
+				[](void* context, bool cancel)
+					{
+
+					TWeakPtr<FPlayFabPartyManager>* WeakPFPartyManagerPtr = static_cast<TWeakPtr<FPlayFabPartyManager>*>(context);
+					TSharedPtr<FPlayFabPartyManager> SharedPFPartyManagerPtr = WeakPFPartyManagerPtr->Pin();
+						if (!SharedPFPartyManagerPtr.IsValid())
+						{
+							delete WeakPFPartyManagerPtr;
+							UE_LOG(LogPlayFabParty, Error, TEXT("[InitPlayFabParty PFInitialize 2] PlayFabManager was destroyed before the async block was executed"));
+							return;
+						}
+
+#if _GRDK_EDITION >= 251000 && !WITH_LEGACY_GDK_FOLDER_STRUCTURE
+						if (!SharedPFPartyManagerPtr->bIsPlayFabPartyInit 
+							|| !SharedPFPartyManagerPtr->bIsPlayFabPartyXboxInit 
+							|| !SharedPFPartyManagerPtr->bIsPlayFabCoreInit)
+#else
+						if (!SharedPFPartyManagerPtr->bIsPlayFabPartyInit 
+							|| !SharedPFPartyManagerPtr->bIsPlayFabPartyXboxInit)
 #endif
-	if (PARTY_FAILED(InitializePlayFabErrorCode))
+
+						{
+							delete WeakPFPartyManagerPtr;
+							SharedPFPartyManagerPtr->ShutdownPlayFabParty();
+							return;
+						}
+
+						SharedPFPartyManagerPtr->TickHandle = FTSTicker::GetCoreTicker().AddTicker(FTickerDelegate::CreateRaw(SharedPFPartyManagerPtr.Get(), &FPlayFabPartyManager::Tick), 0.0f);
+						delete WeakPFPartyManagerPtr;
+					});
+
+			if (!SUCCEEDED(InitPFPartySubmitTaskResult))
+			{
+				delete WeakPFPartyManagerPtr;
+				UE_LOG(LogPlayFabParty, Error, TEXT("PlayFabParty failed to upload Init task to process queue 2. ErrorCode=[%0.8x] ErrorMessage=[%s]"), InitPFPartySubmitTaskResult, *GetPlayFabPartyErrorMessage(InitPFPartySubmitTaskResult));
+			}
+		});
+
+	if (!SUCCEEDED(InitPFPartySubmitTaskResult))
+	{
+		delete WeakThisPtr;
+		UE_LOG(LogPlayFabParty, Error, TEXT("PlayFabParty failed to upload Init task to process queue. ErrorCode=[%0.8x] ErrorMessage=[%s]"), InitPFPartySubmitTaskResult, *GetPlayFabPartyErrorMessage(InitPFPartySubmitTaskResult));
+		
+	}
+}
+
+#if _GRDK_EDITION >= 251000 && !WITH_LEGACY_GDK_FOLDER_STRUCTURE
+void FPlayFabPartyManager::ShutdownPlayFabCore(bool bForceRunSync)
+{
+	if (!ensureMsgf(bIsPlayFabCoreInit, TEXT("PlayFabCore wasn't initialized")) )
 	{
-		UE_LOG(LogPlayFabParty, Warning, TEXT("Failed to initialize PlayFabParty. ErrorCode=[%u] Error=[%s]"), InitializePlayFabErrorCode, *GetPlayFabPartyErrorMessage(InitializePlayFabErrorCode));
-		// We don't try to try again because config is probably invalid if this fails
 		return;
 	}
-	bIsPlayFabPartyInit = true;
 
-	const PartyError InitializePlayFabXboxErrorCode = Party::PartyXblManager::GetSingleton().Initialize(ConvertedAppId.Get());
-	if (PARTY_FAILED(InitializePlayFabXboxErrorCode))
+	GDK_SCOPE_NOT_TIME_SENSITIVE(); // (shutdown only) PFUninitializeAsync cleanup calls many functions which are not safe to call on time-sensitive threads
+	for (TTuple<uint64, PFEntityHandle>& IdToHandle : XboxUserIdToPlayFabEntityHandleMap)
+	{
+		if (IdToHandle.Value)
+		{
+			PFEntityCloseHandle(IdToHandle.Value);
+			IdToHandle.Value = nullptr;
+		}
+		else
+		{
+			UE_LOG(LogPlayFabParty, Warning, TEXT("Failed to cleanup PlayFabCore EntityHandle, invalid handle"));
+		}
+	}
+	XboxUserIdToPlayFabEntityHandleMap.Empty();
+
+	UE_LOG(LogPlayFabParty, Log, TEXT("Shutting down PlayFabCore"));
+	if (PFCoreServiceHandle)
+	{
+		PFServiceConfigCloseHandle(PFCoreServiceHandle);
+		PFCoreServiceHandle = nullptr;
+		UE_LOG(LogPlayFabParty, Log, TEXT("PlayFabCore service config handle cleaned"));
+	}
+	else
 	{
-		UE_LOG(LogPlayFabParty, Warning, TEXT("Failed to initialize PlayFabParty Xbox. ErrorCode=[%u] Error=[%s]"), InitializePlayFabXboxErrorCode, *GetPlayFabPartyErrorMessage(InitializePlayFabXboxErrorCode));
-		// We don't try to try again because config is probably invalid if this fails
+		UE_LOG(LogPlayFabParty, Warning, TEXT("Failed to cleanup PlayFabCore service config, handle was null"));
+	}
+
+	
+
+	if (bForceRunSync)
+	{
+		XAsyncBlock PFCoreShutdownAsyncBlock{};
+
+		HRESULT PFCoreShutdownResult = PFUninitializeAsync(&PFCoreShutdownAsyncBlock);
+		if (!SUCCEEDED(PFCoreShutdownResult))
+		{
+			UE_LOG(LogPlayFabParty, Error, TEXT("PlayFabCore uninitialization sync call failed. ErrorCode=[%0.8x] ErrorMessage=[%s]"), PFCoreShutdownResult, *GetPlayFabPartyErrorMessage(PFCoreShutdownResult));
+			return;
+		}
+
+		PFCoreShutdownResult = XAsyncGetStatus(&PFCoreShutdownAsyncBlock, true);
+		if (SUCCEEDED(PFCoreShutdownResult))
+		{
+			UE_LOG(LogPlayFabParty, Log, TEXT("PlayFabCore uninitialized successfully"));
+		}
+		else
+		{
+			UE_LOG(LogPlayFabParty, Error, TEXT("PlayFabCore uninitialization failed. ErrorCode=[%0.8x] ErrorMessage=[%s]"), PFCoreShutdownResult, *GetPlayFabPartyErrorMessage(PFCoreShutdownResult));
+		}
 
-		ShutdownPlayFabParty();
 		return;
 	}
-	bIsPlayFabPartyXboxInit = true;
 
-	TickHandle = FTSTicker::GetCoreTicker().AddTicker(FTickerDelegate::CreateRaw(this, &FPlayFabPartyManager::Tick), 0.0f);
+	HRESULT PFCoreShutdownResult = AsyncGDKTask(
+		[weakThis = this->AsWeak()](XAsyncBlock* PFCoreShutdownAsyncBlock) -> HRESULT
+		{
+			if (TSharedPtr<FPlayFabPartyManager> PFPartyManagerPtr = weakThis.Pin())
+			{
+				HRESULT PFCoreShutdownResult = PFUninitializeAsync(PFCoreShutdownAsyncBlock);
+				if (!SUCCEEDED(PFCoreShutdownResult))
+				{
+					UE_LOG(LogPlayFabParty, Warning, TEXT("[ShutdownPlayFabCore StartGDKTask] PlayFabCore uninitialization async call failed. ErrorCode=[%0.8x] ErrorMessage=[%s]"), PFCoreShutdownResult, *GetPlayFabPartyErrorMessage(PFCoreShutdownResult));
+				}
+
+				return PFCoreShutdownResult;
+			}
+
+			UE_LOG(LogPlayFabParty, Warning, TEXT("[ShutdownPlayFabCore StartGDKTask] PlayFabManager was destroyed before the async block was executed"));
+			return E_ABORT;
+		},
+		[weakThis = this->AsWeak()](XAsyncBlock* PFCoreLoginAsyncBlock)
+		{
+			TSharedPtr<FPlayFabPartyManager> PFPartyManagerPtr = weakThis.Pin();
+			if (!PFPartyManagerPtr)
+			{
+				UE_LOG(LogPlayFabParty, Warning, TEXT("[ShutdownPlayFabCore GDKTaskCallback] PlayFabManager was destroyed before the async block was executed"));
+				return;
+			}
+
+			UE_LOG(LogPlayFabParty, Log, TEXT("PlayFabCore uninitialized successfully"));
+		}
+	, FGDKAsyncTaskQueue::GetGenericSerialQueue());
+
+	if (!SUCCEEDED(PFCoreShutdownResult))
+	{
+		UE_LOG(LogPlayFabParty, Warning, TEXT("PlayFabCore uninitialization async call failed. ErrorCode=[%0.8x] ErrorMessage=[%s]"), PFCoreShutdownResult, *GetPlayFabPartyErrorMessage(PFCoreShutdownResult));
+	}
+
 }
+#endif
 
-void FPlayFabPartyManager::ShutdownPlayFabParty()
+void FPlayFabPartyManager::CleanupPlayFabPartySingleton()
 {
-	UE_LOG(LogPlayFabParty, Log, TEXT("Shutting down PlayFabParty"));
-	check(bIsPlayFabPartyInit || bIsPlayFabPartyXboxInit);
+	GDK_SCOPE_NOT_TIME_SENSITIVE(); // (shutdown only) PartyManager cleanup calls many functions which are not safe to call on time-sensitive threads
 
-	Logout();
+	// Cleanup PlayFabParty
+	PartyError CleanupErrorCode = Party::PartyManager::GetSingleton().Cleanup();
+	if (PARTY_FAILED(CleanupErrorCode))
+	{
+		UE_LOG(LogPlayFabParty, Warning, TEXT("Failed to cleanup PlayFabParty. ErrorCode=[%u] Error=[%s]"), CleanupErrorCode, *GetPlayFabPartyErrorMessage(CleanupErrorCode));
+		// Still unbind and pretend we shutdown successfully if we fail
+	}
 
-	if (bIsPlayFabPartyInit)
+	UE_LOG(LogPlayFabParty, Log, TEXT("PlayFabParty successfully cleanup."));
+}
+
+void FPlayFabPartyManager::CleanupPlayFabPartyXboxSingleton()
+{
+	GDK_SCOPE_NOT_TIME_SENSITIVE(); // (shutdown only) PartyXblManager cleanup calls many functions which are not safe to call on time-sensitive threads
+
+	// Cleanup PlayFabParty Xbox
+	const PartyError CleanupXboxErrorCode = Party::PartyXblManager::GetSingleton().Cleanup();
+	if (PARTY_FAILED(CleanupXboxErrorCode))
 	{
-		GDK_SCOPE_NOT_TIME_SENSITIVE(); // (shutdown only) PartyManager cleanup calls many functions which are not safe to call on time-sensitive threads
+		UE_LOG(LogPlayFabParty, Warning, TEXT("Failed to cleanup PlayFabParty Xbox. ErrorCode=[%u] Error=[%s]"), CleanupXboxErrorCode, *GetPlayFabPartyErrorMessage(CleanupXboxErrorCode));
+		// Still unbind and pretend we shutdown successfully if we fail
+	}
 
-		// Cleanup PlayFabParty
-		const PartyError CleanupErrorCode = Party::PartyManager::GetSingleton().Cleanup();
-		if (PARTY_FAILED(CleanupErrorCode))
+	UE_LOG(LogPlayFabParty, Log, TEXT("PlayFabParty Xbox successfully cleanup."));
+}
+
+void FPlayFabPartyManager::ShutdownPlayFabParty(bool bForceRunSync)
+{
+	UE_LOG(LogPlayFabParty, Log, TEXT("Shutting down PlayFabParty"));
+
+	if (bForceRunSync)
+	{
+#if _GRDK_EDITION >= 251000 && !WITH_LEGACY_GDK_FOLDER_STRUCTURE
+		check(bIsPlayFabPartyInit || bIsPlayFabPartyXboxInit || bIsPlayFabCoreInit);
+		ShutdownPlayFabCore(bForceRunSync);
+		bIsPlayFabCoreInit = false;
+#else
+		check(bIsPlayFabPartyInit || bIsPlayFabPartyXboxInit);
+#endif
+
+		Logout();
+
+		if (bIsPlayFabPartyInit)
 		{
-			UE_LOG(LogPlayFabParty, Warning, TEXT("Failed to cleanup PlayFabParty. ErrorCode=[%u] Error=[%s]"), CleanupErrorCode, *GetPlayFabPartyErrorMessage(CleanupErrorCode));
-			// Still unbind and pretend we shutdown successfully if we fail
+			CleanupPlayFabPartySingleton();
+			bIsPlayFabPartyInit = false;
 		}
-		bIsPlayFabPartyInit = false;
-	}
 
-	if (bIsPlayFabPartyXboxInit)
+		if (bIsPlayFabPartyXboxInit)
+		{
+			CleanupPlayFabPartyXboxSingleton();
+			bIsPlayFabPartyXboxInit = false;
+		}
+	}
+	else
 	{
-		GDK_SCOPE_NOT_TIME_SENSITIVE(); // (shutdown only) PartyXblManager cleanup calls many functions which are not safe to call on time-sensitive threads
 
-		// Cleanup PlayFabParty Xbox
-		const PartyError CleanupXboxErrorCode = Party::PartyXblManager::GetSingleton().Cleanup();
-		if (PARTY_FAILED(CleanupXboxErrorCode))
+#if _GRDK_EDITION >= 251000 && !WITH_LEGACY_GDK_FOLDER_STRUCTURE
+		check(bIsPlayFabPartyInit || bIsPlayFabPartyXboxInit || bIsPlayFabCoreInit);
+		ShutdownPlayFabCore(bForceRunSync);
+		bIsPlayFabCoreInit = false;
+#else
+		check(bIsPlayFabPartyInit || bIsPlayFabPartyXboxInit);
+#endif
+
+		Logout();
+
+		if (bIsPlayFabPartyInit)
 		{
-			UE_LOG(LogPlayFabParty, Warning, TEXT("Failed to cleanup PlayFabParty Xbox. ErrorCode=[%u] Error=[%s]"), CleanupXboxErrorCode, *GetPlayFabPartyErrorMessage(CleanupXboxErrorCode));
-			// Still unbind and pretend we shutdown successfully if we fail
+			HRESULT ShutdownPFPartySubmitTaskResult = XTaskQueueSubmitCallback(
+				FGDKAsyncTaskQueue::GetGenericSerialQueue(),
+				XTaskQueuePort::Work,
+				nullptr,
+				[](void* context, bool cancel)
+				{
+					CleanupPlayFabPartySingleton();
+				});
+
+			bIsPlayFabPartyInit = false; // We notify that PlayFabParty is not init even before the callback completed to prevent usage during shutdown, we pretend that is shutdown even if the internal task fails.
+
+			if (!SUCCEEDED(ShutdownPFPartySubmitTaskResult))
+			{
+				UE_LOG(LogPlayFabParty, Error, TEXT("PlayFabParty failed to upload Shutdown task to process queue. ErrorCode=[%0.8x] ErrorMessage=[%s]"), ShutdownPFPartySubmitTaskResult, *GetPlayFabPartyErrorMessage(ShutdownPFPartySubmitTaskResult));
+			}
+		}
+
+		if (bIsPlayFabPartyXboxInit)
+		{
+
+			HRESULT ShutdownPFPartyXboxSubmitTaskResult = XTaskQueueSubmitCallback(
+				FGDKAsyncTaskQueue::GetGenericSerialQueue(),
+				XTaskQueuePort::Work,
+				nullptr,
+				[](void* context, bool cancel)
+				{
+					CleanupPlayFabPartyXboxSingleton();
+				});
+
+			bIsPlayFabPartyXboxInit = false; // We notify that PlayFabPartyXbox is not init even before the callback completed to prevent usage during shutdown, we pretend that is shutdown even if the internal task fails.
+
+			if (!SUCCEEDED(ShutdownPFPartyXboxSubmitTaskResult))
+			{
+				UE_LOG(LogPlayFabParty, Warning, TEXT("PlayFabPartyXbox failed to upload Shutdown task to process queue. ErrorCode=[%0.8x] ErrorMessage=[%s]"), ShutdownPFPartyXboxSubmitTaskResult, *GetPlayFabPartyErrorMessage(ShutdownPFPartyXboxSubmitTaskResult));
+			}
 		}
-		bIsPlayFabPartyXboxInit = false;
 	}
 
 	if (TickHandle.IsValid())
@@ -416,6 +722,7 @@ void FPlayFabPartyManager::ShutdownPlayFabParty()
 	}
 
 	QOSData.Reset();
+	
 }
 
 void FPlayFabPartyManager::CheckLoginStatus()
@@ -430,6 +737,9 @@ void FPlayFabPartyManager::CheckLoginStatus()
 
 	if (LocalUserLoginState == EPlayFabPartyLoginState::LoggedIn)
 	{
+
+// Token refresh is handled automatically after 251000
+#if _GRDK_EDITION < 251000 || WITH_LEGACY_GDK_FOLDER_STRUCTURE
 		// Make sure we actually need to refresh our login
 		if (CurrentTimeSeconds < NextLoginRefreshTimeSeconds)
 		{
@@ -442,6 +752,7 @@ void FPlayFabPartyManager::CheckLoginStatus()
 
 		// Refresh our login credentials
 		RefreshLogin();
+#endif
 	}
 	else
 	{
@@ -533,7 +844,9 @@ void FPlayFabPartyManager::Logout()
 
 	LocalUserXuid = 0;
 	LocalUserLoginState = EPlayFabPartyLoginState::NeedLogin;
+#if _GRDK_EDITION < 251000 || WITH_LEGACY_GDK_FOLDER_STRUCTURE // Token refresh is handled automatically after 251000
 	NextLoginRefreshTimeSeconds = 0.0;
+#endif
 }
 
 void FPlayFabPartyManager::HandleAppSuspended()
@@ -1057,51 +1370,28 @@ void FPlayFabPartyManager::HandleXboxCreateLocalChatUserCompleted(const Party::P
 	}
 }
 
-void FPlayFabPartyManager::HandleLoginToPlayFabCompleted(const Party::PartyXblLoginToPlayFabCompletedStateChange* const StateChange)
+#if _GRDK_EDITION < 251000 || WITH_LEGACY_GDK_FOLDER_STRUCTURE
+void FPlayFabPartyManager::HandleLocalPlayerCreationWithEntityToken(const Party::PartyXblLoginToPlayFabCompletedStateChange* const StateChange)
 {
-	UE_LOG(LogPlayFabParty, Verbose, TEXT("Received PlayFabParty Xbox LoginToPlayFabCompleted(%u) event"), StateChange->stateChangeType);
-
-	if (StateChange->result != Party::PartyXblStateChangeResult::Succeeded)
-	{
-		UE_LOG(LogPlayFabParty, Warning, TEXT("PlayFabParty login failed. Result=[%u] ErrorCode=[%u] Error=[%s]"), StateChange->result, StateChange->errorDetail, *GetPlayFabPartyErrorMessage(StateChange->errorDetail));
-
-		NextLoginTimeSeconds = FPlatformTime::Seconds() + LoginFailureDelaySeconds;
-		Logout();
-		return;
-	}
-
-	// Cache our local entity id
-	if (LocalUserXuid > 0 && XboxUserIdToPlayFabEntityIdMap.Find(LocalUserXuid) == nullptr)
-	{
-		XboxUserIdToPlayFabEntityIdMap.Emplace(LocalUserXuid, StateChange->entityId);
-	}
-
-	XboxUserIdToPlayFabEntityTokenMap.Emplace(LocalUserXuid, StateChange->titlePlayerEntityToken);
-
 	if (LocalUserLoginState == EPlayFabPartyLoginState::LoggedIn && PartyLocalUser)
 	{
-		UE_LOG(LogPlayFabParty, Log, TEXT("Successfully refreshed login token for %") UINT64_FMT, LocalUserXuid);
-
 		// Token refresh
 		const PartyError UpdateTokenError = PartyLocalUser->UpdateEntityToken(StateChange->titlePlayerEntityToken);
 		if (PARTY_FAILED(UpdateTokenError))
 		{
 			UE_LOG(LogPlayFabParty, Warning, TEXT("Failed to update our PlayFabParty local user token. ErrorCode=[%u] Error=[%s]"), UpdateTokenError, *GetPlayFabPartyErrorMessage(UpdateTokenError));
 		}
+		else
+		{
+			UE_LOG(LogPlayFabParty, Log, TEXT("Successfully refreshed login token for %") UINT64_FMT, LocalUserXuid);
+		}
 	}
 	else if (LocalUserLoginState == EPlayFabPartyLoginState::InProgress)
 	{
 		// New Login
 		UE_LOG(LogPlayFabParty, Log, TEXT("Successfully logged into PlayFabParty for %") UINT64_FMT, LocalUserXuid);
 
-#if _GRDK_EDITION < 251000 || WITH_LEGACY_GDK_FOLDER_STRUCTURE
 		const PartyError CreateLocalUserError = Party::PartyManager::GetSingleton().CreateLocalUser(StateChange->entityId, StateChange->titlePlayerEntityToken, &PartyLocalUser);
-#else
-		// NOT IMPLEMENTED. @fixme: #jira UE-315672
-		// entitly handle should be obtained via PFAuthenticationGetEntityWithSecretKeyAsync, PFAuthenticationLoginWithXboxAsync or PFAuthenticationLoginWithXUserAsync		 
-		checkNoEntry();
-		const PartyError CreateLocalUserError = Party::PartyManager::GetSingleton().CreateLocalUser(PFEntityHandle(), &PartyLocalUser);
-#endif
 		if (PARTY_FAILED(CreateLocalUserError))
 		{
 			UE_LOG(LogPlayFabParty, Warning, TEXT("Failed to create our PlayFabParty local user. ErrorCode=[%u] Error=[%s]"), CreateLocalUserError, *GetPlayFabPartyErrorMessage(CreateLocalUserError));
@@ -1115,6 +1405,125 @@ void FPlayFabPartyManager::HandleLoginToPlayFabCompleted(const Party::PartyXblLo
 		LocalUserLoginState = EPlayFabPartyLoginState::LoggedIn;
 	}
 }
+#endif
+
+#if _GRDK_EDITION >= 251000 && !WITH_LEGACY_GDK_FOLDER_STRUCTURE
+void FPlayFabPartyManager::HandleLocalPlayerCreationWithEntityHandle()
+{
+	if (LocalUserLoginState == EPlayFabPartyLoginState::InProgress)
+	{
+		// New Login
+		UE_LOG(LogPlayFabParty, Log, TEXT("Successfully logged into PlayFabParty for %") UINT64_FMT, LocalUserXuid);
+
+		HRESULT CreateLocalUserResult = AsyncGDKTask(
+			[weakThis = this->AsWeak()](XAsyncBlock* PFCoreLoginAsyncBlock) -> HRESULT
+			{
+				if (TSharedPtr<FPlayFabPartyManager> PFPartyManagerPtr = weakThis.Pin())
+				{
+					FGDKUserHandle UserHandle = IGDKRuntimeModule::Get().GetUserHandleByXUserId(PFPartyManagerPtr->LocalUserXuid);
+					PFAuthenticationLoginWithXUserRequest LoginRequest{};
+					LoginRequest.createAccount = true;
+					LoginRequest.user = UserHandle;
+					HRESULT Result = PFAuthenticationLoginWithXUserAsync(PFPartyManagerPtr->PFCoreServiceHandle, &LoginRequest, PFCoreLoginAsyncBlock);
+
+					return Result;
+				}
+
+				UE_LOG(LogPlayFabParty, Warning, TEXT("[PlayFabCoreCreateLocalUser StartGDKTask] PlayFabManager was destroyed before the async block was executed"));
+				return E_ABORT;
+			},
+			[weakThis = this->AsWeak()](XAsyncBlock* PFCoreLoginAsyncBlock)
+			{
+				TSharedPtr<FPlayFabPartyManager> PFPartyManagerPtr = weakThis.Pin();
+				if (!PFPartyManagerPtr)
+				{
+					UE_LOG(LogPlayFabParty, Warning, TEXT("[PlayFabCoreCreateLocalUser GDKTaskCallback] PlayFabManager was destroyed before the async block was executed"));
+					return;
+				}
+
+				if (PFPartyManagerPtr->LocalUserLoginState !=  EPlayFabPartyLoginState::InProgress)
+				{
+					UE_LOG(LogPlayFabParty, Log, TEXT("[PlayFabCoreCreateLocalUser GDKTaskCallback] Current user was logged out before login operation completed"));
+					return;
+				}
+
+				std::vector<char> LoginResultBuffer;
+				PFAuthenticationLoginResult const* LoginResult;
+				size_t BufferSize;
+				HRESULT PFCoreLoginResult = PFAuthenticationLoginWithXUserGetResultSize(PFCoreLoginAsyncBlock, &BufferSize);
+				if (FAILED(PFCoreLoginResult))
+				{
+					UE_LOG(LogPlayFabParty, Warning, TEXT("PlayFabCore failed to get login result size. ErrorCode=[%0.8x] ErrorMessage=[%s]"), PFCoreLoginResult, *GetPlayFabPartyErrorMessage(PFCoreLoginResult));
+					PFPartyManagerPtr->NextLoginTimeSeconds = FPlatformTime::Seconds() + PFPartyManagerPtr->LoginFailureDelaySeconds;
+					PFPartyManagerPtr->Logout();
+					return;
+				}
+
+				LoginResultBuffer.resize(BufferSize);
+				PFEntityHandle EntityHandle{ nullptr };
+				PFCoreLoginResult = PFAuthenticationLoginWithXUserGetResult(PFCoreLoginAsyncBlock, &EntityHandle, LoginResultBuffer.size(), LoginResultBuffer.data(), &LoginResult, nullptr);
+				if (FAILED(PFCoreLoginResult))
+				{
+					UE_LOG(LogPlayFabParty, Warning, TEXT("PlayFabCore failed to get login result data. ErrorCode=[%0.8x] ErrorMessage=[%s]"), PFCoreLoginResult, *GetPlayFabPartyErrorMessage(PFCoreLoginResult));
+					PFPartyManagerPtr->NextLoginTimeSeconds = FPlatformTime::Seconds() + PFPartyManagerPtr->LoginFailureDelaySeconds;
+					PFPartyManagerPtr->Logout();
+					return;
+				}
+
+				const PartyError CreateLocalUserError = Party::PartyManager::GetSingleton().CreateLocalUser(EntityHandle, &PFPartyManagerPtr->PartyLocalUser);
+				if (PARTY_FAILED(CreateLocalUserError))
+				{
+					UE_LOG(LogPlayFabParty, Warning, TEXT("PlayFabCore failed to create local user. ErrorCode=[%0.8x] ErrorMessage=[%s]"), CreateLocalUserError, *GetPlayFabPartyErrorMessage(CreateLocalUserError));
+					PFPartyManagerPtr->NextLoginTimeSeconds = FPlatformTime::Seconds() + PFPartyManagerPtr->LoginFailureDelaySeconds;
+					PFPartyManagerPtr->Logout();
+					return;
+				}
+
+				UE_LOG(LogPlayFabParty, Log, TEXT("PlayFabCore login operation SUCCEEDED for XBOX user. PlayerID -> %s"), *FString(LoginResult->playFabId));
+
+				PFPartyManagerPtr->XboxUserIdToPlayFabEntityHandleMap.Add(PFPartyManagerPtr->LocalUserXuid, EntityHandle);
+				PFPartyManagerPtr->LocalUserLoginState = EPlayFabPartyLoginState::LoggedIn;
+			}
+		, FGDKAsyncTaskQueue::GetGenericSerialQueue());
+
+		if (FAILED(CreateLocalUserResult))
+		{
+			UE_LOG(LogPlayFabParty, Warning, TEXT("PlayFabCore failed to start login process for user. ErrorCode=[%0.8x] ErrorMessage=[%s]"), CreateLocalUserResult, *GetPlayFabPartyErrorMessage(CreateLocalUserResult));
+			NextLoginTimeSeconds = FPlatformTime::Seconds() + LoginFailureDelaySeconds;
+			Logout();
+		}
+	}
+}
+#endif
+
+void FPlayFabPartyManager::HandleLoginToPlayFabCompleted(const Party::PartyXblLoginToPlayFabCompletedStateChange* const StateChange)
+{
+	UE_LOG(LogPlayFabParty, Verbose, TEXT("Received PlayFabParty Xbox LoginToPlayFabCompleted(%u) event"), StateChange->stateChangeType);
+
+	if (StateChange->result != Party::PartyXblStateChangeResult::Succeeded)
+	{
+		UE_LOG(LogPlayFabParty, Warning, TEXT("PlayFabParty login failed. Result=[%u] ErrorCode=[%u] Error=[%s]"), StateChange->result, StateChange->errorDetail, *GetPlayFabPartyErrorMessage(StateChange->errorDetail));
+
+		NextLoginTimeSeconds = FPlatformTime::Seconds() + LoginFailureDelaySeconds;
+		Logout();
+		return;
+	}
+
+	// Cache our local entity id
+	if (LocalUserXuid > 0 && XboxUserIdToPlayFabEntityIdMap.Find(LocalUserXuid) == nullptr)
+	{
+		XboxUserIdToPlayFabEntityIdMap.Emplace(LocalUserXuid, StateChange->entityId);
+	}
+
+	// In GDK >= 251000 this EntityToken is the same one as the one retrieved using EntityHandle after creating a local user. We are good storing this one even before getting an EntityHandle.
+	XboxUserIdToPlayFabEntityTokenMap.Emplace(LocalUserXuid, StateChange->titlePlayerEntityToken); 
+
+#if _GRDK_EDITION < 251000 || WITH_LEGACY_GDK_FOLDER_STRUCTURE
+	HandleLocalPlayerCreationWithEntityToken(StateChange);
+#else
+	HandleLocalPlayerCreationWithEntityHandle();
+#endif
+}
 
 void FPlayFabPartyManager::HandleGetEntityIdsFromXboxLiveUserIdsCompleted(const Party::PartyXblGetEntityIdsFromXboxLiveUserIdsCompletedStateChange* const StateChange)
 {
diff --git a/Engine/Plugins/Online/Microsoft/PlayFabParty/Source/Private/PlayFabPartyManager.h b/Engine/Plugins/Online/Microsoft/PlayFabParty/Source/Private/PlayFabPartyManager.h
--- a/Engine/Plugins/Online/Microsoft/PlayFabParty/Source/Private/PlayFabPartyManager.h
+++ b/Engine/Plugins/Online/Microsoft/PlayFabParty/Source/Private/PlayFabPartyManager.h
@@ -7,9 +7,13 @@
 #include "PlayFabParty.h"
 #include "Containers/Ticker.h"
 
+THIRD_PARTY_INCLUDES_START
 #include "Microsoft/AllowMicrosoftPlatformTypes.h"
+#include <grdk.h>
 #include <XTaskQueue.h>
 #include "Microsoft/HideMicrosoftPlatformTypes.h"
+THIRD_PARTY_INCLUDES_END
+
 
 namespace Party
 {
@@ -78,6 +82,7 @@ private:
 
 };
 
+
 /** Class to manage the status of PlayFabParty and handle its events */
 class FPlayFabPartyManager
 	: public TSharedFromThis<FPlayFabPartyManager, ESPMode::ThreadSafe>
@@ -124,6 +129,11 @@ public:
 	/** Get the cached Entity ID for the provided Xuid if any! */
 	PartyString GetEntityTokenForXuid(const uint64 Xuid);
 
+#if _GRDK_EDITION >= 251000 && !WITH_LEGACY_GDK_FOLDER_STRUCTURE
+	/** Get the cached Entity Handle for the provided Xuid if any! */
+	PFEntityHandle GetEntityHandleForXuid(const uint64 Xuid);
+#endif
+
 	/** Get PlayFab socket of listen server */
 	TSharedPtr<FSocket> GetPlayFabSocketListenServer() const { return PlayFabSocketListenServer; }
 	/** Save PlayFab socket of listen server so to prevent recreating */
@@ -140,8 +150,17 @@ protected:
 
 	/** Attempt to initialize PlayFabParty */
 	void InitPlayFabParty();
+
+#if _GRDK_EDITION >= 251000 && !WITH_LEGACY_GDK_FOLDER_STRUCTURE
+	/** Shutdown PlayFabCore */
+	void ShutdownPlayFabCore(bool bForceRunSync = false);
+#endif
+	/** Shutdown PlayFabParty manager singleton */
+	static void CleanupPlayFabPartySingleton();
+	/** Shutdown PlayFabParty Xbox manager singleton */
+	static void CleanupPlayFabPartyXboxSingleton();
 	/** Shutdown PlayFabParty */
-	void ShutdownPlayFabParty();
+	void ShutdownPlayFabParty(bool bForceRunSync = false);
 
 	/** Check if we should login or refresh our auth right now */
 	void CheckLoginStatus();
@@ -185,19 +204,34 @@ protected:
 	void HandleGetEntityIdsFromXboxLiveUserIdsCompleted(const Party::PartyXblGetEntityIdsFromXboxLiveUserIdsCompletedStateChange* const StateChange);
 
 protected:
+#if _GRDK_EDITION < 251000 || WITH_LEGACY_GDK_FOLDER_STRUCTURE
+	void HandleLocalPlayerCreationWithEntityToken(const Party::PartyXblLoginToPlayFabCompletedStateChange* const);
+#else
+	void HandleLocalPlayerCreationWithEntityHandle();
+#endif
+		
 	/** Reference to the socket subsystem that created us */
 	FPlayFabPartySocketSubsystem& SocketSubsystem;
 	/** The AppId for this application from PlayFab's website */
 	FString AppId;
+	/** The URL for this application PlayFab's website */
+	FString PlayFabURL;
 	/** The delay in seconds after a login completes (success or failure) before another login may occur */
 	double LoginFailureDelaySeconds = 15.0;
+
+#if _GRDK_EDITION < 251000 || WITH_LEGACY_GDK_FOLDER_STRUCTURE
 	/** The delay in seconds after a successful login complete before we should login again to get a new token */
 	double LoginRefreshDelaySeconds = 7200.0;
+#endif
 
+#if _GRDK_EDITION >= 251000 && !WITH_LEGACY_GDK_FOLDER_STRUCTURE
+	/** Is PlayFabCore currently initialized? */
+	std::atomic<bool> bIsPlayFabCoreInit = false;
+#endif
 	/** Is PlayFabParty currently initialized? */
-	bool bIsPlayFabPartyInit = false;
+	std::atomic<bool> bIsPlayFabPartyInit = false;
 	/** Is PlayFabParty Xbox currently initialized? */
-	bool bIsPlayFabPartyXboxInit = false;
+	std::atomic<bool> bIsPlayFabPartyXboxInit = false;
 
 	/** Is the app currently suspended? */
 	bool bIsAppSuspending = false;
@@ -207,6 +241,11 @@ protected:
 	/** The flag to indicate if the game is connected playfab network. Since it's async operation, we can't create another socket to connect to playfab network, while previous socket hasn't finished leaving playfab network yet */
 	bool bIsNetworkConnected = false;
 
+#if _GRDK_EDITION >= 251000 && !WITH_LEGACY_GDK_FOLDER_STRUCTURE
+	/** Handle to manage PlayFabCore lifecycle */
+	std::atomic<PFServiceConfigHandle> PFCoreServiceHandle {};
+#endif
+	
 	/** Handle to our QOS Data binding */
 	FDelegateHandle QOSDataRequestedHandle;
 
@@ -225,8 +264,11 @@ protected:
 	EPlayFabPartyLoginState LocalUserLoginState = EPlayFabPartyLoginState::NeedLogin;
 	/** The next time we can possibly try logging in again (throttle protection for login requests) */
 	double NextLoginTimeSeconds = 0.0;
+
+#if _GRDK_EDITION < 251000 || WITH_LEGACY_GDK_FOLDER_STRUCTURE // Token refresh is handled automatically after 251000
 	/** The next time we should refresh our login credentials (to ensure they do not expire while in use) */
 	double NextLoginRefreshTimeSeconds = 0.0;
+#endif
 
 	/** Our local xbox chat user */
 	Party::PartyXblLocalChatUser* PartyXboxLocalChatUser = nullptr;
@@ -247,5 +289,10 @@ protected:
 
 	/** Map of Xbox XUIDs to PlayFab Entity Tokens */
 	TMap<uint64, TUniqueObj<FPlayFabPartyEntityToken>> XboxUserIdToPlayFabEntityTokenMap;
+
+#if _GRDK_EDITION >= 251000 && !WITH_LEGACY_GDK_FOLDER_STRUCTURE
+	/** Map of Xbox XUIDs to PlayFab Entity Handles */
+	TMap<uint64, PFEntityHandle> XboxUserIdToPlayFabEntityHandleMap;
+#endif
 };
 #endif // WITH_PLAYFAB_PARTY
diff --git a/Engine/Source/Runtime/Microsoft/GDKRuntime/Private/GDKTaskQueueHelpers.cpp b/Engine/Source/Runtime/Microsoft/GDKRuntime/Private/GDKTaskQueueHelpers.cpp
--- a/Engine/Source/Runtime/Microsoft/GDKRuntime/Private/GDKTaskQueueHelpers.cpp
+++ b/Engine/Source/Runtime/Microsoft/GDKRuntime/Private/GDKTaskQueueHelpers.cpp
@@ -8,6 +8,7 @@
 #include "HAL/PlatformProcess.h"
 #include "HAL/PlatformMisc.h"
 #include "Containers/Ticker.h"
+#include "Tasks/Pipe.h"
 #include "Templates/UniquePtr.h"
 
 FGDKAsyncBlock::FGDKAsyncBlock(void* InUserData, FGDKAsyncBlockDelegate InDelegate, XTaskQueueHandle TaskQueue)
@@ -94,8 +95,44 @@ private:
 	FTSTicker::FDelegateHandle TickHandle;
 };
 
+class FThreadPoolSerialTaskQueue : public FGDKAsyncTaskQueue
+{
+public:
+	FThreadPoolSerialTaskQueue()
+		: FGDKAsyncTaskQueue(XTaskQueueDispatchMode::Manual, XTaskQueueDispatchMode::Manual)
+	{
+		// set up to dispatch completions once a frame on game thread
+		TickHandle = FTSTicker::GetCoreTicker().AddTicker(FTickerDelegate::CreateLambda([this](float)
+			{
+				// Try to launch one completion task per frame as done in FThreadPoolSerialTaskQueue, if the game thread queue is empty we launch a work task in the thread pool.
+				if (!XTaskQueueDispatch(Queue, XTaskQueuePort::Completion, 0)) 
+				{
+					if (!WorkTasksPipe.HasWork())
+					{
+						WorkTasksPipe.Launch(
+							TEXT("GDK Dispatch queue bg task"),
+							[this]
+							{
+								XTaskQueueDispatch(Queue, XTaskQueuePort::Work, 0);
+							},
+							UE::Tasks::ETaskPriority::Normal
+						);
+					}
+				}
 
+				return true;
+			}));
+	}
+
+	~FThreadPoolSerialTaskQueue()
+	{
+		FTSTicker::GetCoreTicker().RemoveTicker(TickHandle);
+	}
 
+private:
+	FTSTicker::FDelegateHandle TickHandle;
+	UE::Tasks::FPipe WorkTasksPipe{TEXT("GDK Serial Task Queue")};
+};
 
 class FBackgroundTaskQueue : public FGDKAsyncTaskQueue
 {
@@ -112,6 +149,7 @@ public:
 
 static bool bIsTornDown = false;
 static TUniquePtr<FThreadPoolTaskQueue> GThreadPoolTaskQueue;
+static TUniquePtr<FThreadPoolSerialTaskQueue> GThreadPoolSerialTaskQueue;
 static TUniquePtr<FBackgroundTaskQueue> GBackgroundTaskQueue;
 
 void FGDKAsyncTaskQueue::PlatformInit()
@@ -128,6 +166,7 @@ void FGDKAsyncTaskQueue::PlatformTeardown()
 
 	// destroy our task queues
 	GThreadPoolTaskQueue.Reset();
+	GThreadPoolSerialTaskQueue.Reset();
 	GBackgroundTaskQueue.Reset();
 
 	// close the system task queue and prevent anyone from using it
@@ -153,6 +192,18 @@ XTaskQueueHandle FGDKAsyncTaskQueue::GetGenericQueue()
 	return GThreadPoolTaskQueue->GetQueue();
 }
 
+XTaskQueueHandle FGDKAsyncTaskQueue::GetGenericSerialQueue()
+{
+	check(!bIsTornDown);
+
+	if (!GThreadPoolSerialTaskQueue.IsValid())
+	{
+		GThreadPoolSerialTaskQueue = MakeUnique<FThreadPoolSerialTaskQueue>();
+	}
+
+	return GThreadPoolSerialTaskQueue->GetQueue();
+}
+
 XTaskQueueHandle FGDKAsyncTaskQueue::GetBackgroundTaskQueue()
 {
 	check(!bIsTornDown);
diff --git a/Engine/Source/Runtime/Microsoft/GDKRuntime/Public/GDKTaskQueueHelpers.h b/Engine/Source/Runtime/Microsoft/GDKRuntime/Public/GDKTaskQueueHelpers.h
--- a/Engine/Source/Runtime/Microsoft/GDKRuntime/Public/GDKTaskQueueHelpers.h
+++ b/Engine/Source/Runtime/Microsoft/GDKRuntime/Public/GDKTaskQueueHelpers.h
@@ -65,6 +65,7 @@ public:
 
 	// access to common task queues
 	static UE_API XTaskQueueHandle GetGenericQueue();			// completion callbacks happen on the game thread
+	static UE_API XTaskQueueHandle GetGenericSerialQueue();		// completion callbacks happen on the game thread and are guaranteed to execute before the next work task starts in a bg thread, ensuring Work -> Callback order
 	static UE_API XTaskQueueHandle GetBackgroundTaskQueue();	// completion callbacks happen on a background thread
 
 	// internal use - startup & shutdown the queue system
