fix for PlayFab on PC GDK building with the new folder layout

#rnx

#Synced-CL 51493042

diff --git a/Engine/Plugins/Online/Microsoft/OnlineSubsystemGDK/Source/OnlineSubsystemGDK.Build.cs b/Engine/Plugins/Online/Microsoft/OnlineSubsystemGDK/Source/OnlineSubsystemGDK.Build.cs
--- a/Engine/Plugins/Online/Microsoft/OnlineSubsystemGDK/Source/OnlineSubsystemGDK.Build.cs
+++ b/Engine/Plugins/Online/Microsoft/OnlineSubsystemGDK/Source/OnlineSubsystemGDK.Build.cs
@@ -88,6 +88,11 @@ public class OnlineSubsystemGDK : ModuleRules
 				}
 			}
 
+			if (Target.Platform.IsInGroup(UnrealPlatformGroup.Windows))
+			{
+				PublicDefinitions.Add("_GAMING_DESKTOP");
+			}
+
 			// Add OSS configuration files
 			RuntimeDependencies.Add("$(ProjectDir)/Platforms/GDK/Config/OSS/*", StagedFileType.UFS);
 		}
diff --git a/Engine/Plugins/Online/Microsoft/PlayFabParty/Source/PlayFabParty.build.cs b/Engine/Plugins/Online/Microsoft/PlayFabParty/Source/PlayFabParty.build.cs
--- a/Engine/Plugins/Online/Microsoft/PlayFabParty/Source/PlayFabParty.build.cs
+++ b/Engine/Plugins/Online/Microsoft/PlayFabParty/Source/PlayFabParty.build.cs
@@ -52,6 +52,10 @@ public class PlayFabParty : ModuleRules
 		}
 		PublicDefinitions.Add("WITH_PLAYFAB_PARTY=" + (bHasPlayFabParty ? "1" : "0"));
 
+		if (Target.Platform.IsInGroup(UnrealPlatformGroup.Windows))
+		{
+			PublicDefinitions.Add("_GAMING_DESKTOP");
+		}
 
 		//CDA duplicated from OnlineSubsytemGDK to avoid circular dependency.
 		PublicDefinitions.Add("GDK_SUBSYSTEM=FName(TEXT(\"GDK\"))");
